Adding support for additional Turnkey manifests in the Engine/Platforms folders.

#rnx

diff --git a/Engine/Build/Turnkey/TurnkeyManifest.xml b/Engine/Build/Turnkey/TurnkeyManifest.xml
--- a/Engine/Build/Turnkey/TurnkeyManifest.xml
+++ b/Engine/Build/Turnkey/TurnkeyManifest.xml
@@ -1,6 +1,7 @@
 <?xml version="1.0" encoding="utf-8" ?>
 <TurnkeyManifest>
 	<AdditionalManifests>
+		<Manifest>file:$(EngineDir)/Platforms/*/Build/Turnkey/TurnkeyManifest.xml</Manifest>
 		<Manifest>$(ThisManifestDir)/Licensee_TurnkeyManifest.xml</Manifest>
 		<Manifest>file:$(EngineDir)/Restricted/NotForLicensees/Build/Turnkey/TurnkeyManifest.xml</Manifest>
 		<Manifest>file:$(EngineDir)/Restricted/NoRedist/Build/Turnkey/TurnkeyManifest.xml</Manifest>
