- Added some example TurnkeyManifest xml files for licensees to use as a starting point for their own hosting

diff --git a/Engine/Build/Turnkey/Licensee_TurnkeyManifest.Example.xml b/Engine/Build/Turnkey/Licensee_TurnkeyManifest.Example.xml
new file mode 100644
--- /dev/null
+++ b/Engine/Build/Turnkey/Licensee_TurnkeyManifest.Example.xml
@@ -0,0 +1,69 @@
+<?xml version="1.0" encoding="utf-8" ?>
+<TurnkeyManifest>
+
+	<FileSource>
+		<Platform>listexpansion:ExpPlatform=$(AutoSDKPlatforms)</Platform>
+		<Version>$(ExpVersion)</Version>
+		<Name>$(ExpPlatform) AutoSdk version $(ExpVersion)</Name>
+		<Type>AutoSdk</Type>
+		<!-- Because p4 can't get just directories, we don't want to grab the entire ... recursive mess of files when expanding, we look for a single file (well-known .bat and .sh) in each directory
+		     but then we CopyOverride to grab ... when one is actually chosen. If we had zipped AutoSDKs on googledrive or something, we could make this much simpler -->
+		<Source CopyOverride="perforce://depot/AutoSDK/Host$(HOST_PLATFORM_NAME)/$(ExpPlatform)/$(Version)/...">
+			fileexpansion:perforce://depot/AutoSDK/Host$(HOST_PLATFORM_NAME)/$(ExpPlatform)/$[ExpVersion]/Setup.*
+		</Source>
+	</FileSource>
+
+	<FileSource>
+		<!-- Android is pretty special in that it has a script for setting up the SDK, but we can still use Turnkey to grab AndroidStudio, which is a pre-requiste, not an SDK -->
+		<Platform>Android</Platform>
+		<Type>Misc</Type>
+		<Name>AndroidStudio</Name>
+		<Source HostPlatform="Win64">googledrive:/SdkInstallers/Android/android-studio-ide-193.6821437-windows.exe</Source>
+		<Source HostPlatform="Linux">googledrive:/SdkInstallers/Android/android-studio-ide-193.6821437-linux.tar.gz</Source>
+		<Source HostPlatform="Mac">googledrive:/SdkInstallers/Android/android-studio-ide-193.6821437-mac.zip</Source>
+	</FileSource>
+
+	
+	<!-- A lot of the IOS stuff below is a work in progress, especially the code signing stuff -->
+	<FileSource>
+		<Platform>Mac,IOS,TVOS</Platform>
+		<Version>$(ExpVersion)</Version>
+		<Type>Full</Type>
+		<Name>Xcode $(ExpVersion)</Name>
+		<Source HostPlatform="Mac">fileexpansion:googledrive:/SdkInstallers/Apple/Installers/Xcode_$[ExpVersion].xip</Source>
+	</FileSource>
+
+	<FileSource>
+		<Platform>IOS,TVOS</Platform>
+		<!-- Http can't enumerate files, so fileexpansion won't work to find version numbers, like others can, so hardcode a value here (this is a dummy version) -->
+		<Version>1300.0</Version>
+		<Type>Full</Type>
+		<Name>Latest iTunes</Name>
+		<Source HostPlatform="Win64">https://www.apple.com/itunes/download/win64</Source>
+	</FileSource>
+
+	<FileSource>
+		<Platform>IOS,TVOS</Platform>
+		<Version>$(ExpVersion)</Version>
+		<Type>Flash</Type>
+		<AllowedFlashDeviceTypes>$(ExpDeviceType)</AllowedFlashDeviceTypes>
+		<Name>IOS Flash $(ExpVersion) for $(ExpDeviceType)</Name>
+		<Source >fileexpansion:googledrive:/SdkInstallers/IOS/Flashes/$[ExpVersion]/$[ExpDeviceType]_$[ExpVersion].ipsw</Source>
+	</FileSource>
+
+	<FileSource>
+		<Type>Misc</Type>
+		<!-- Here we use a regex to match on Name when finding a cert for a particular appid or range of appids -->
+		<Name>regex:DevCert: com\.mycompany\..*</Name>
+		<Source HostPlatform="Win64">file:\\server.net\IOS\Certs\SpecificDevCert.p12</Source>
+		<Source HostPlatform="Mac">file:smb://server.net/IOS/Certs/SpecificDevCert.p12</Source>
+	</FileSource>
+
+	<FileSource>
+		<Type>Misc</Type>
+		<Name>DevCert</Name>
+		<Source HostPlatform="Win64">file:\\server.net\IOS\Certs\GenericDevCert.p12</Source>
+		<Source HostPlatform="Mac">file:smb://server.net/IOS/Certs/GenericDevCert.p12</Source>
+	</FileSource>
+
+</TurnkeyManifest>
