- Removed AndroidStudio installer for non-existant "MacArm" platform until we have a solution for arm-specific installers

diff --git a/Engine/Build/Turnkey/TurnkeyManifest.xml b/Engine/Build/Turnkey/TurnkeyManifest.xml
--- a/Engine/Build/Turnkey/TurnkeyManifest.xml
+++ b/Engine/Build/Turnkey/TurnkeyManifest.xml
@@ -18,7 +18,6 @@
         <Source HostPlatform="Win64">https://redirector.gvt1.com/edgedl/android/studio/install/2022.2.1.20/android-studio-2022.2.1.20-windows.exe</Source>
         <Source HostPlatform="Linux">https://redirector.gvt1.com/edgedl/android/studio/ide-zips/2022.2.1.20/android-studio-2022.2.1.20-linux.tar.gz</Source>
         <Source HostPlatform="Mac">https://redirector.gvt1.com/edgedl/android/studio/install/2022.2.1.20/android-studio-2022.2.1.20-mac.dmg</Source>
-        <Source HostPlatform="MacArm">https://redirector.gvt1.com/edgedl/android/studio/install/2022.2.1.20/android-studio-2022.2.1.20-mac_arm.dmg</Source>
     </FileSource>
 
 </TurnkeyManifest>
