Enable Android Turnkey for launcher builds by default (only public external dependencies)

#android

diff --git a/Engine/Build/Turnkey/TurnkeyStudioSettings.xml b/Engine/Build/Turnkey/TurnkeyStudioSettings.xml
--- a/Engine/Build/Turnkey/TurnkeyStudioSettings.xml
+++ b/Engine/Build/Turnkey/TurnkeyStudioSettings.xml
@@ -17,7 +17,7 @@
 		<SavedSetting>
 			<!-- By default, indicate that Turnkey cannot install any platforms. This can be a comma separated list, or All to indicate all platforms should try to install via Turnkey. -->
 			<Variable>Studio_FullInstallPlatforms</Variable>
-			<Value></Value>
+			<Value>Android</Value>
 		</SavedSetting>
 		<SavedSetting>
 			<Variable>Studio_AutoSdkPlatforms</Variable>
