Update to new Android Studio version for Turnkey

#android

diff --git a/Engine/Build/Turnkey/TurnkeyManifest.xml b/Engine/Build/Turnkey/TurnkeyManifest.xml
--- a/Engine/Build/Turnkey/TurnkeyManifest.xml
+++ b/Engine/Build/Turnkey/TurnkeyManifest.xml
@@ -14,10 +14,10 @@
         <Platform>Android</Platform>
         <Type>Misc</Type>
         <Name>AndroidStudio</Name>
-		<!-- Flamingo 2022.2.1 Patch 2 May 24, 2023 -->
-        <Source HostPlatform="Win64">https://redirector.gvt1.com/edgedl/android/studio/install/2022.2.1.20/android-studio-2022.2.1.20-windows.exe</Source>
-        <Source HostPlatform="Linux">https://redirector.gvt1.com/edgedl/android/studio/ide-zips/2022.2.1.20/android-studio-2022.2.1.20-linux.tar.gz</Source>
-        <Source HostPlatform="Mac">https://redirector.gvt1.com/edgedl/android/studio/install/2022.2.1.20/android-studio-2022.2.1.20-mac.dmg</Source>
+		<!--  Koala Feature Drop 2024.1.2 Patch 1 September 17, 2024 -->
+        <Source HostPlatform="Win64">https://redirector.gvt1.com/edgedl/android/studio/install/2024.1.2.13/android-studio-2024.1.2.13-windows.exe</Source>
+        <Source HostPlatform="Linux">https://redirector.gvt1.com/edgedl/android/studio/ide-zips/2024.1.2.13/android-studio-2024.1.2.13-linux.tar.gz</Source>
+        <Source HostPlatform="Mac">https://redirector.gvt1.com/edgedl/android/studio/install/2024.1.2.13/android-studio-2024.1.2.13-mac.dmg</Source>
     </FileSource>
 
 </TurnkeyManifest>
